SourceControlWindowExtender: Extends functionality of the SourceControlWindows window

Actor: Focus, Pin, Select
Assets: Browse To

diff --git a/Engine/Source/Editor/SceneOutliner/Private/ActorDescTreeItem.cpp b/Engine/Source/Editor/SceneOutliner/Private/ActorDescTreeItem.cpp
--- a/Engine/Source/Editor/SceneOutliner/Private/ActorDescTreeItem.cpp
+++ b/Engine/Source/Editor/SceneOutliner/Private/ActorDescTreeItem.cpp
@@ -270,13 +270,8 @@ void FActorDescTreeItem::FocusActorBounds() const
 {
 	if (FWorldPartitionActorDesc const* ActorDesc = ActorDescHandle.Get())
 	{
-		if (FLevelEditorViewportClient* LevelViewportClient = GCurrentLevelEditingViewportClient)
-		{
-			FBox Bounds = ActorDesc->GetBounds();
-			LevelViewportClient->FocusViewportOnBox(Bounds, false);
-			const FColor DarkGray(40, 40, 40);
-			DrawDebugBox(ActorDesc->GetContainer()->GetWorld(), Bounds.GetCenter(), Bounds.GetExtent(), DarkGray, false, 0.5f);
-		}
+		const bool bActiveViewportOnly = true;
+		GEditor->MoveViewportCamerasToBox(ActorDesc->GetBounds(), bActiveViewportOnly, 0.5f);
 	}
 }
 
diff --git a/Engine/Source/Editor/SourceControlWindowExtender/Private/SourceControlWindowExtenderModule.cpp b/Engine/Source/Editor/SourceControlWindowExtender/Private/SourceControlWindowExtenderModule.cpp
new file mode 100644
--- /dev/null
+++ b/Engine/Source/Editor/SourceControlWindowExtender/Private/SourceControlWindowExtenderModule.cpp
@@ -0,0 +1,302 @@
+// Copyright Epic Games, Inc. All Rights Reserved.
+
+#include "CoreMinimal.h"
+#include "Modules/ModuleManager.h"
+#include "ISourceControlWindowExtenderModule.h"
+#include "ISourceControlWindowsModule.h"
+#include "SourceControlMenuContext.h"
+#include "SourceControlHelpers.h"
+#include "ToolMenus.h"
+#include "AssetRegistry/AssetData.h"
+#include "GameFramework/Actor.h"
+#include "Editor.h"
+#include "Selection.h"
+#include "ScopedTransaction.h"
+#include "WorldPartition/WorldPartition.h"
+#include "WorldPartition/WorldPartitionActorDesc.h"
+#include "WorldPartition/WorldPartitionActorDescUtils.h"
+
+#define LOCTEXT_NAMESPACE "SourceControlWindowExtender"
+
+/**
+ * SourceControlWindowExtender module
+ */
+class FSourceControlWindowExtenderModule : public ISourceControlWindowExtenderModule
+{
+public:
+	/**
+	 * Called right after the module DLL has been loaded and the module object has been created
+	 */
+	virtual void StartupModule() override;
+
+	/**
+	 * Called before the module is unloaded, right before the module object is destroyed.
+	 */
+	virtual void ShutdownModule() override;
+
+private:
+	void ExtendMenu();
+
+	bool CanPinActors() const { return CurrentWorldUnloadedActors.Num() > 0; }
+	void PinActors();
+	void PinActors(const TArray<FAssetData>& ActorsToPin);
+	bool CanSelectActors() const { return CurrentWorldLoadedActors.Num() > 0; }
+	void SelectActors();
+	void SelectActors(const TArray<FAssetData>& ActorsToSelect);
+	void BrowseToAssets();
+	void BrowseToAssets(const TArray<FAssetData>& Assets);
+	void FocusActors(const TArray<FAssetData>& ActorToFocus, bool bSelect);
+	bool CanFocusActors() const { return CurrentWorldLoadedActors.Num() > 0 || CurrentWorldUnloadedActors.Num() > 0; }
+	void FocusActors();
+
+	void OnChangelistFileDoubleClicked(const FString& Filename);
+
+	void GetAssetsFromFilenames(const TArray<FString>& Filenames, TArray<FAssetData>& OutNonActorAssets, TArray<FAssetData>& OutCurrentWorldLoadedActors, TArray<FAssetData>& OutCurrentWorldUnloadedActors);
+
+	TArray<FAssetData> SelectedAssets;
+	TArray<FAssetData> CurrentWorldLoadedActors;
+	TArray<FAssetData> CurrentWorldUnloadedActors;
+	FDelegateHandle EventHandle;
+};
+
+IMPLEMENT_MODULE(FSourceControlWindowExtenderModule, SourceControlWindowExtender);
+
+void FSourceControlWindowExtenderModule::StartupModule()
+{
+	ExtendMenu();
+
+	EventHandle = ISourceControlWindowsModule::Get().OnChangelistFileDoubleClicked().AddRaw(this, &FSourceControlWindowExtenderModule::OnChangelistFileDoubleClicked);
+}
+
+void FSourceControlWindowExtenderModule::ShutdownModule()
+{
+	if (ISourceControlWindowsModule* Module = ISourceControlWindowsModule::TryGet())
+	{
+		Module->OnChangelistFileDoubleClicked().Remove(EventHandle);
+	}
+}
+
+void FSourceControlWindowExtenderModule::GetAssetsFromFilenames(const TArray<FString>& Filenames, TArray<FAssetData>& OutNonActorAssets, TArray<FAssetData>& OutCurrentWorldLoadedActors, TArray<FAssetData>& OutCurrentWorldUnloadedActors)
+{
+	UWorld* CurrentWorld = GEditor->GetEditorWorldContext().World();
+
+	for (const FString& Filename : Filenames)
+	{
+		TArray<FAssetData> OutAssets;
+		if (SourceControlHelpers::GetAssetData(Filename, OutAssets) && OutAssets.Num() == 1)
+		{
+			const FAssetData& AssetData = OutAssets[0];
+			if (TSubclassOf<AActor> ActorClass = AssetData.GetClass())
+			{
+				if (CurrentWorld && AssetData.ObjectPath.ToString().StartsWith(CurrentWorld->GetPathName()))
+				{
+					if (AssetData.IsAssetLoaded())
+					{
+						OutCurrentWorldLoadedActors.Add(AssetData);
+					}
+					else
+					{
+						OutCurrentWorldUnloadedActors.Add(AssetData);
+					}
+				}
+				else
+				{
+					TArray<FAssetData> OutWorldAsset;
+					FString AssetPathName = AssetData.ToSoftObjectPath().GetLongPackageName();
+					if (SourceControlHelpers::GetAssetDataFromPackage(AssetPathName, OutWorldAsset) && OutWorldAsset.Num() == 1)
+					{
+						OutNonActorAssets.Add(OutWorldAsset[0]);
+					}
+				}
+			}
+			else
+			{
+				OutNonActorAssets.Add(AssetData);
+			}
+		}
+	}
+}
+
+void FSourceControlWindowExtenderModule::ExtendMenu()
+{
+	UToolMenu* Menu = UToolMenus::Get()->ExtendMenu("SourceControl.ChangelistContextMenu");
+	FToolMenuSection& Section = Menu->AddDynamicSection("SourceControlWindowExtenderSection", FNewToolMenuDelegate::CreateLambda([this](UToolMenu* InMenu)
+	{
+		if (USourceControlMenuContext* MenuContext = InMenu->FindContext<USourceControlMenuContext>())
+        {
+            SelectedAssets.Empty();
+            CurrentWorldLoadedActors.Empty();
+            CurrentWorldUnloadedActors.Empty();
+
+			GetAssetsFromFilenames(MenuContext->SelectedFiles, SelectedAssets, CurrentWorldLoadedActors, CurrentWorldUnloadedActors);
+
+			auto FindOrAddExtendSection = [InMenu]() -> FToolMenuSection*
+			{
+				FToolMenuSection* ExtendSection = InMenu->FindSection("SourceControlWindowExtender");
+				if (!ExtendSection)
+				{
+					ExtendSection = &InMenu->AddSection("SourceControlWindowExtender", TAttribute<FText>(), FToolMenuInsert("Source Control", EToolMenuInsertType::After));
+					ExtendSection->AddSeparator(NAME_None);
+				}
+				return ExtendSection;
+			};
+
+            if (CurrentWorldUnloadedActors.Num() > 0 || CurrentWorldLoadedActors.Num() > 0)
+            {
+				FindOrAddExtendSection()->AddSubMenu(TEXT("Actors"),
+                    LOCTEXT("ActorSubMenu", "Actors"),
+                    LOCTEXT("ActorSubMenuTooltip", ""),
+                    FNewToolMenuChoice(FNewMenuDelegate::CreateLambda([this](FMenuBuilder& MenuBuilder)
+                    {
+                        MenuBuilder.AddMenuEntry(LOCTEXT("PinActors", "Pin"), LOCTEXT("PinActors_Tooltip", "Load actors"),
+							FSlateIcon(), FUIAction(FExecuteAction::CreateRaw(this, &FSourceControlWindowExtenderModule::PinActors), FCanExecuteAction::CreateRaw(this, &FSourceControlWindowExtenderModule::CanPinActors)));
+                        MenuBuilder.AddMenuEntry(LOCTEXT("SelectActors", "Select"), LOCTEXT("SelectActors_Tooltip", "Select actors"), 
+							FSlateIcon(), FUIAction(FExecuteAction::CreateRaw(this, &FSourceControlWindowExtenderModule::SelectActors), FCanExecuteAction::CreateRaw(this, &FSourceControlWindowExtenderModule::CanSelectActors)));
+						MenuBuilder.AddMenuEntry(LOCTEXT("FocusActors", "Focus"), LOCTEXT("FocusActors_Tooltip", "Focus actors"),
+							FSlateIcon(), FUIAction(FExecuteAction::CreateRaw(this, &FSourceControlWindowExtenderModule::FocusActors), FCanExecuteAction::CreateRaw(this, &FSourceControlWindowExtenderModule::CanFocusActors)));
+                    })));
+            }
+
+            if (SelectedAssets.Num() > 0)
+            {
+				FindOrAddExtendSection()->AddSubMenu(TEXT("Assets"),
+                    LOCTEXT("AssetSubMenu", "Assets"),
+                    LOCTEXT("AssetSubMenuTooltip", ""),
+                    FNewToolMenuChoice(FNewMenuDelegate::CreateLambda([this](FMenuBuilder& MenuBuilder)
+                    {
+                        MenuBuilder.BeginSection(NAME_None);
+                        MenuBuilder.AddMenuEntry(LOCTEXT("BrowseToAssets", "Browse to Asset"), LOCTEXT("BrowseToAssets_Tooltip", "Browse to Asset in Content Browser"), FSlateIcon(FAppStyle::GetAppStyleSetName(), "SystemWideCommands.FindInContentBrowser.Small"), FUIAction(FExecuteAction::CreateRaw(this, &FSourceControlWindowExtenderModule::BrowseToAssets)));
+                        MenuBuilder.EndSection();
+                    })));
+            }
+        }
+    }));
+}
+
+void FSourceControlWindowExtenderModule::PinActors()
+{
+	PinActors(CurrentWorldUnloadedActors);
+}
+
+void FSourceControlWindowExtenderModule::PinActors(const TArray<FAssetData>& ActorsToPin)
+{
+	UWorld* CurrentWorld = GEditor->GetEditorWorldContext().World();
+	check(CurrentWorld);
+	if (UWorldPartition* WorldPartition = CurrentWorld->GetWorldPartition())
+	{
+		TArray<FGuid> ActorGuids;
+		for (const FAssetData& ActorToPin : ActorsToPin)
+		{
+			if(TUniquePtr<FWorldPartitionActorDesc> ActorDesc = FWorldPartitionActorDescUtils::GetActorDescriptorFromAssetData(ActorToPin))
+			{
+				ActorGuids.Add(ActorDesc->GetGuid());
+			}
+		}
+
+		if (!ActorGuids.IsEmpty())
+		{
+			WorldPartition->PinActors(ActorGuids);
+			SelectActors(CurrentWorldUnloadedActors);
+		}
+	}
+}
+
+void FSourceControlWindowExtenderModule::SelectActors()
+{
+	SelectActors(CurrentWorldLoadedActors);
+}
+
+void FSourceControlWindowExtenderModule::SelectActors(const TArray<FAssetData>& ActorsToSelect)
+{
+	const FScopedTransaction Transaction(LOCTEXT("SelectActorsFromChangelist", "Select Actor(s)"));
+	UWorld* CurrentWorld = GEditor->GetEditorWorldContext().World();
+	check(CurrentWorld);
+
+	GEditor->GetSelectedActors()->BeginBatchSelectOperation();
+	bool bNotify = false;
+	const bool bDeselectBSPSurfs = true;
+	GEditor->SelectNone(bNotify, bDeselectBSPSurfs);
+	
+	for (const FAssetData& ActorToSelect : ActorsToSelect)
+	{
+		if (AActor* Actor = Cast<AActor>(ActorToSelect.FastGetAsset()))
+		{
+			const bool bSelected = true;
+			GEditor->SelectActor(Actor, bSelected, bNotify);
+		}
+	}
+
+	bNotify = true;
+	GEditor->GetSelectedActors()->EndBatchSelectOperation(bNotify);
+}
+
+void FSourceControlWindowExtenderModule::BrowseToAssets()
+{
+	BrowseToAssets(SelectedAssets);
+}
+
+void FSourceControlWindowExtenderModule::BrowseToAssets(const TArray<FAssetData>& Assets)
+{
+	GEditor->SyncBrowserToObjects(const_cast<TArray<FAssetData>&>(Assets));
+}
+
+void FSourceControlWindowExtenderModule::FocusActors()
+{
+	TArray<FAssetData> ActorsToFocus;
+	ActorsToFocus.Append(CurrentWorldLoadedActors);
+	ActorsToFocus.Append(CurrentWorldUnloadedActors);
+
+	const bool bSelectActors = false;
+	FocusActors(ActorsToFocus, bSelectActors);
+}
+
+void FSourceControlWindowExtenderModule::FocusActors(const TArray<FAssetData>& ActorsToFocus, bool bSelectActors)
+{
+	FBox FocusBounds(EForceInit::ForceInit);
+	UWorld* CurrentWorld = GEditor->GetEditorWorldContext().World();
+	check(CurrentWorld);
+	for (const FAssetData& ActorToFocus : ActorsToFocus)
+	{
+		if (TUniquePtr<FWorldPartitionActorDesc> ActorDesc = FWorldPartitionActorDescUtils::GetActorDescriptorFromAssetData(ActorToFocus))
+		{
+			FocusBounds += ActorDesc->GetBounds();
+		}
+	}
+
+	if (FocusBounds.IsValid)
+	{
+		const bool bActiveViewportOnly = true;
+		const float TimeInSeconds = 0.5f;
+		GEditor->MoveViewportCamerasToBox(FocusBounds, bActiveViewportOnly, TimeInSeconds);
+	}
+
+	if (bSelectActors)
+	{
+		SelectActors(ActorsToFocus);
+	}
+}
+
+void FSourceControlWindowExtenderModule::OnChangelistFileDoubleClicked(const FString& Filename)
+{
+	TArray<FAssetData> OutNonActorAssets;
+	TArray<FAssetData> OutCurrentWorldLoadedActors;
+	TArray<FAssetData> OutCurrentWorldUnloadedActors;
+
+	GetAssetsFromFilenames({ Filename }, OutNonActorAssets, OutCurrentWorldLoadedActors, OutCurrentWorldUnloadedActors);
+
+	if (OutNonActorAssets.Num() > 0)
+	{
+		BrowseToAssets(OutNonActorAssets);
+	}
+	else if (OutCurrentWorldUnloadedActors.Num() > 0)
+	{
+		const bool bSelectActors = false;
+		FocusActors(OutCurrentWorldUnloadedActors, bSelectActors);
+	}
+	else if (OutCurrentWorldLoadedActors.Num() > 0)
+	{
+		const bool bSelectActors = true;
+		FocusActors(OutCurrentWorldLoadedActors, bSelectActors);
+	}
+}
+
diff --git a/Engine/Source/Editor/SourceControlWindowExtender/Public/ISourceControlWindowExtenderModule.h b/Engine/Source/Editor/SourceControlWindowExtender/Public/ISourceControlWindowExtenderModule.h
new file mode 100644
--- /dev/null
+++ b/Engine/Source/Editor/SourceControlWindowExtender/Public/ISourceControlWindowExtenderModule.h
@@ -0,0 +1,21 @@
+// Copyright Epic Games, Inc. All Rights Reserved.
+#pragma once
+
+#include "CoreMinimal.h"
+#include "Modules/ModuleInterface.h"
+#include "Modules/ModuleManager.h"
+
+/**
+ * SourceControlWindowExtender module interface
+ */
+class ISourceControlWindowExtenderModule : public IModuleInterface
+{
+public:
+	/**
+	 * Get reference to the SourceControlWindowExtender module instance
+	 */
+	static inline ISourceControlWindowExtenderModule& Get()
+	{
+		return FModuleManager::LoadModuleChecked<ISourceControlWindowExtenderModule>("SourceControlWindowExtender");
+	}
+};
diff --git a/Engine/Source/Editor/SourceControlWindowExtender/SourceControlWindowExtender.Build.cs b/Engine/Source/Editor/SourceControlWindowExtender/SourceControlWindowExtender.Build.cs
new file mode 100644
--- /dev/null
+++ b/Engine/Source/Editor/SourceControlWindowExtender/SourceControlWindowExtender.Build.cs
@@ -0,0 +1,27 @@
+// Copyright Epic Games, Inc. All Rights Reserved.
+
+using UnrealBuildTool;
+
+public class SourceControlWindowExtender : ModuleRules
+{
+	public SourceControlWindowExtender(ReadOnlyTargetRules Target) : base(Target)
+	{
+		PublicIncludePaths.Add("Editor/SourceControlWindowExtender/Public");
+
+		PrivateIncludePaths.Add("Editor/SourceControlWindowExtender/Private");
+
+		PrivateDependencyModuleNames.AddRange(
+			new string[] {
+				"Core", 
+				"CoreUObject", 
+				"Engine",
+            	"ToolMenus",
+				"SlateCore",
+				"Slate",
+				"SourceControl",
+				"SourceControlWindows",
+				"UnrealEd"
+			}
+		);
+	}
+}
diff --git a/Engine/Source/Editor/SourceControlWindows/Private/SSourceControlChangelists.cpp b/Engine/Source/Editor/SourceControlWindows/Private/SSourceControlChangelists.cpp
--- a/Engine/Source/Editor/SourceControlWindows/Private/SSourceControlChangelists.cpp
+++ b/Engine/Source/Editor/SourceControlWindows/Private/SSourceControlChangelists.cpp
@@ -19,6 +19,7 @@
 
 #include "ISourceControlProvider.h"
 #include "ISourceControlModule.h"
+#include "ISourceControlWindowsModule.h"
 #include "UncontrolledChangelistsModule.h"
 #include "SourceControlOperations.h"
 #include "ToolMenus.h"
@@ -26,6 +27,7 @@
 #include "SourceControlWindows.h"
 #include "SourceControlHelpers.h"
 #include "SourceControlPreferences.h"
+#include "SourceControlMenuContext.h"
 #include "AssetToolsModule.h"
 #include "ContentBrowserModule.h"
 #include "IContentBrowserSingleton.h"
@@ -873,22 +875,31 @@ void SSourceControlChangelistsWidget::GetSelectedFiles(TArray<FString>& OutContr
 
 	for (const TSharedPtr<IChangelistTreeItem>& Item : SelectedItems)
 	{
-		if (Item->GetTreeItemType() != IChangelistTreeItem::File)
-		{
-			continue;
-		}
-
-		if (TSharedPtr<IChangelistTreeItem> Parent = Item->GetParent())
+		if (Item->GetTreeItemType() == IChangelistTreeItem::File)
 		{
-			const FString& Filename = StaticCastSharedPtr<FFileTreeItem>(Item)->FileState->GetFilename();
-
-			if (Parent->GetTreeItemType() == IChangelistTreeItem::Changelist)
+			if (TSharedPtr<IChangelistTreeItem> Parent = Item->GetParent())
 			{
-				OutControlledFiles.Add(Filename);
+				const FString& Filename = StaticCastSharedPtr<FFileTreeItem>(Item)->FileState->GetFilename();
+
+				if (Parent->GetTreeItemType() == IChangelistTreeItem::Changelist)
+				{
+					OutControlledFiles.Add(Filename);
+				}
+				else if (Parent->GetTreeItemType() == IChangelistTreeItem::UncontrolledChangelist)
+				{
+					OutUncontrolledFiles.Add(Filename);
+				}
 			}
-			else if (Parent->GetTreeItemType() == IChangelistTreeItem::UncontrolledChangelist)
+		}
+		else if (Item->GetTreeItemType() == IChangelistTreeItem::OfflineFile)
+		{
+			if (TSharedPtr<IChangelistTreeItem> Parent = Item->GetParent())
 			{
-				OutUncontrolledFiles.Add(Filename);
+				if (Parent->GetTreeItemType() == IChangelistTreeItem::UncontrolledChangelist)
+				{
+					const FString& Filename = StaticCastSharedPtr<FOfflineFileTreeItem>(Item)->GetFilename();
+					OutUncontrolledFiles.Add(Filename);
+				}
 			}
 		}
 	}
@@ -1714,51 +1725,6 @@ void SSourceControlChangelistsWidget::OnMoveFiles()
 	}
 }
 
-void SSourceControlChangelistsWidget::OnLocateFile()
-{
-	TArray<FAssetData> AssetsToSync;
-	TArray<FChangelistTreeItemPtr> SelectedItems = FileTreeView->GetSelectedItems();
-
-	for (const FChangelistTreeItemPtr& SelectedItem : SelectedItems)
-	{
-		if (SelectedItem->GetTreeItemType() == IChangelistTreeItem::File)
-		{
-			const FAssetDataArrayPtr& Assets = StaticCastSharedPtr<FFileTreeItem>(SelectedItem)->GetAssetData();
-
-			if (Assets.IsValid())
-			{
-				AssetsToSync.Append(*Assets);
-			}
-		}
-	}
-
-	if (AssetsToSync.Num() > 0)
-	{
-		FContentBrowserModule& ContentBrowserModule = FModuleManager::Get().LoadModuleChecked<FContentBrowserModule>("ContentBrowser");
-		ContentBrowserModule.Get().SyncBrowserToAssets(AssetsToSync, true);
-	}
-}
-
-bool SSourceControlChangelistsWidget::CanLocateFile()
-{
-	TArray<FChangelistTreeItemPtr> SelectedItems = FileTreeView->GetSelectedItems();
-
-	auto HasAssetData = [](const FChangelistTreeItemPtr& SelectedItem)
-	{
-		if (SelectedItem->GetTreeItemType() != IChangelistTreeItem::File)
-		{
-			return false;
-		}
-
-		const FAssetDataArrayPtr& Assets = StaticCastSharedPtr<FFileTreeItem>(SelectedItem)->GetAssetData();
-
-		return (Assets.IsValid() && Assets->Num() > 0);
-	};
-
-	// Checks if at least one selected item has asset data (ie: accessible from ContentBrowser)
-	return SelectedItems.FindByPredicate(HasAssetData) != nullptr;
-}
-
 void SSourceControlChangelistsWidget::OnShowHistory()
 {
 	TArray<FString> SelectedFiles = GetSelectedFiles();
@@ -1802,13 +1768,11 @@ TSharedPtr<SWidget> SSourceControlChangelistsWidget::OnOpenContextMenu()
 	static const FName MenuName = "SourceControl.ChangelistContextMenu";
 	if (!ToolMenus->IsMenuRegistered(MenuName))
 	{
-		ToolMenus->RegisterMenu(MenuName);
+		UToolMenu* RegisteredMenu = ToolMenus->RegisterMenu(MenuName);
+		// Add section so it can be used as insert position for menu extensions
+		RegisteredMenu->AddSection("Source Control");
 	}
 
-	// Build up the menu for a selection
-	FToolMenuContext Context;
-	UToolMenu* Menu = ToolMenus->GenerateMenu(MenuName, Context);
-
 	TArray<TSharedPtr<IChangelistTreeItem>> SelectedChangelistNodes = ChangelistTreeView->GetSelectedItems();
 	TArray<TSharedPtr<IChangelistTreeItem>> SelectedUncontrolledChangelistNodes = UncontrolledChangelistTreeView->GetSelectedItems();
 
@@ -1819,7 +1783,20 @@ TSharedPtr<SWidget> SSourceControlChangelistsWidget::OnOpenContextMenu()
 	bool bHasSelectedShelvedFiles = (GetSelectedShelvedFiles().Num() > 0);
 	bool bHasEmptySelection = (!bHasSelectedChangelist && !bHasSelectedFiles && !bHasSelectedShelvedFiles);
 
-	FToolMenuSection& Section = Menu->AddSection("Source Control");
+	// Build up the menu for a selection
+	USourceControlMenuContext* ContextObject = NewObject<USourceControlMenuContext>();
+	FToolMenuContext Context(ContextObject);
+
+	// Fill Context Object
+	TArray<FString> SelectedControlledFiles;
+	TArray<FString> SelectedUncontrolledFiles;
+	GetSelectedFiles(SelectedControlledFiles, SelectedUncontrolledFiles);
+	ContextObject->SelectedFiles.Append(SelectedControlledFiles);
+	ContextObject->SelectedFiles.Append(SelectedUncontrolledFiles);
+
+	UToolMenu* Menu = ToolMenus->GenerateMenu(MenuName, Context);
+
+	FToolMenuSection& Section = *Menu->FindSection("Source Control");
 	
 	// This should appear only on change lists
 	if (bHasSelectedChangelist)
@@ -1943,15 +1920,6 @@ TSharedPtr<SWidget> SSourceControlChangelistsWidget::OnOpenContextMenu()
 			FSlateIcon(),
 			FUIAction(FExecuteAction::CreateSP(this, &SSourceControlChangelistsWidget::OnMoveFiles)));
 
-		Section.AddMenuEntry(
-			"LocateFile",
-			LOCTEXT("SourceControl_LocateFile", "Locate File..."),
-			LOCTEXT("SourceControl_LocateFile_Tooltip", "Locate File in Project..."),
-			FSlateIcon(),
-			FUIAction(
-				FExecuteAction::CreateSP(this, &SSourceControlChangelistsWidget::OnLocateFile),
-				FCanExecuteAction::CreateSP(this, &SSourceControlChangelistsWidget::CanLocateFile)));
-
 		Section.AddMenuEntry(
 			"ShowHistory",
 			LOCTEXT("SourceControl_ShowHistory", "Show History..."),
@@ -2005,6 +1973,7 @@ TSharedRef<STreeView<FChangelistTreeItemPtr>> SSourceControlChangelistsWidget::C
 		.OnGetChildren(this, &SSourceControlChangelistsWidget::OnGetFileChildren)
 		.SelectionMode(ESelectionMode::Multi)
 		.OnContextMenuOpening(this, &SSourceControlChangelistsWidget::OnOpenContextMenu)
+		.OnMouseButtonDoubleClick(this, &SSourceControlChangelistsWidget::OnItemDoubleClicked)
 		.HeaderRow
 		(
 			SNew(SHeaderRow)
@@ -2617,6 +2586,20 @@ void SSourceControlChangelistsWidget::OnGetChangelistChildren(FChangelistTreeIte
 	}
 }
 
+void SSourceControlChangelistsWidget::OnItemDoubleClicked(TSharedPtr<IChangelistTreeItem> Item)
+{
+	if (Item->GetTreeItemType() == IChangelistTreeItem::OfflineFile)
+	{
+		const FString& Filename = StaticCastSharedPtr<FOfflineFileTreeItem>(Item)->GetFilename();
+		ISourceControlWindowsModule::Get().OnChangelistFileDoubleClicked().Broadcast(Filename);
+	}
+	else if (Item->GetTreeItemType() == IChangelistTreeItem::File)
+	{
+		const FString& Filename = StaticCastSharedPtr<FFileTreeItem>(Item)->FileState->GetFilename();
+		ISourceControlWindowsModule::Get().OnChangelistFileDoubleClicked().Broadcast(Filename);
+	}
+}
+
 void SSourceControlChangelistsWidget::OnChangelistSelectionChanged(TSharedPtr<IChangelistTreeItem> SelectedChangelist, ESelectInfo::Type SelectionType)
 {
 	FileTreeNodes.Reset();
diff --git a/Engine/Source/Editor/SourceControlWindows/Private/SSourceControlChangelists.h b/Engine/Source/Editor/SourceControlWindows/Private/SSourceControlChangelists.h
--- a/Engine/Source/Editor/SourceControlWindows/Private/SSourceControlChangelists.h
+++ b/Engine/Source/Editor/SourceControlWindows/Private/SSourceControlChangelists.h
@@ -109,8 +109,6 @@ private:
 
 	/** Files operations */
 	void OnMoveFiles();
-	void OnLocateFile();
-	bool CanLocateFile();
 	void OnShowHistory();
 	void OnDiffAgainstDepot();
 	bool CanDiffAgainstDepot();
@@ -122,6 +120,7 @@ private:
 	/** Source control callbacks */
 	void OnSourceControlProviderChanged(ISourceControlProvider& OldProvider, ISourceControlProvider& NewProvider);
 	void OnSourceControlStateChanged();
+	void OnItemDoubleClicked(TSharedPtr<IChangelistTreeItem> Item);
 	void OnChangelistSelectionChanged(TSharedPtr<IChangelistTreeItem> SelectedItem, ESelectInfo::Type SelectionType);
 	void OnChangelistsStatusUpdated(const TSharedRef<ISourceControlOperation>& InOperation, ECommandResult::Type InType);
 
diff --git a/Engine/Source/Editor/SourceControlWindows/Private/SSourceControlCommon.cpp b/Engine/Source/Editor/SourceControlWindows/Private/SSourceControlCommon.cpp
--- a/Engine/Source/Editor/SourceControlWindows/Private/SSourceControlCommon.cpp
+++ b/Engine/Source/Editor/SourceControlWindows/Private/SSourceControlCommon.cpp
@@ -224,6 +224,7 @@ FText FShelvedChangelistTreeItem::GetDisplayText() const
 FOfflineFileTreeItem::FOfflineFileTreeItem(const FString& InFilename)
 	: IChangelistTreeItem(IChangelistTreeItem::OfflineFile)
 	, Assets()
+	, Filename(InFilename)
 	, PackageName(FText::FromString(InFilename)) 
 	, AssetName(SSourceControlCommon::GetDefaultAssetName())
 	, AssetPath()
diff --git a/Engine/Source/Editor/SourceControlWindows/Private/SSourceControlCommon.h b/Engine/Source/Editor/SourceControlWindows/Private/SSourceControlCommon.h
--- a/Engine/Source/Editor/SourceControlWindows/Private/SSourceControlCommon.h
+++ b/Engine/Source/Editor/SourceControlWindows/Private/SSourceControlCommon.h
@@ -250,6 +250,8 @@ struct FOfflineFileTreeItem : public IChangelistTreeItem
 	explicit FOfflineFileTreeItem(const FString& InFilename);
 
 public:
+	const FString& GetFilename() const { return Filename; }
+
 	const FText& GetPackageName() const { return PackageName; }
 
 	const FText& GetDisplayName() const { return AssetName; }
@@ -262,6 +264,7 @@ public:
 
 private:
 	TArray<FAssetData> Assets;
+	FString Filename;
 	FText PackageName;
 	FText AssetName;
 	FText AssetPath;
diff --git a/Engine/Source/Editor/SourceControlWindows/Private/SourceControlWindowsModule.cpp b/Engine/Source/Editor/SourceControlWindows/Private/SourceControlWindowsModule.cpp
--- a/Engine/Source/Editor/SourceControlWindows/Private/SourceControlWindowsModule.cpp
+++ b/Engine/Source/Editor/SourceControlWindows/Private/SourceControlWindowsModule.cpp
@@ -37,6 +37,9 @@ public:
 
 	virtual void ShowChangelistsTab() override;
 	virtual bool CanShowChangelistsTab() const override;
+	
+	DECLARE_DERIVED_EVENT(FSourceControlWindowsModule, ISourceControlWindowsModule::FChangelistFileDoubleClickedEvent, FChangelistFileDoubleClickedEvent);
+	virtual FChangelistFileDoubleClickedEvent& OnChangelistFileDoubleClicked() override { return ChangelistFileDoubleClickedEvent; }
 
 private:
 	TSharedRef<SDockTab> CreateChangelistsTab(const FSpawnTabArgs& Args);
@@ -45,6 +48,8 @@ private:
 private:
 	TWeakPtr<SDockTab> ChangelistsTab;
 	TWeakPtr<SSourceControlChangelistsWidget> ChangelistsWidget;
+
+	FChangelistFileDoubleClickedEvent ChangelistFileDoubleClickedEvent;
 };
 
 IMPLEMENT_MODULE(FSourceControlWindowsModule, SourceControlWindows);
diff --git a/Engine/Source/Editor/SourceControlWindows/Public/ISourceControlWindowsModule.h b/Engine/Source/Editor/SourceControlWindows/Public/ISourceControlWindowsModule.h
--- a/Engine/Source/Editor/SourceControlWindows/Public/ISourceControlWindowsModule.h
+++ b/Engine/Source/Editor/SourceControlWindows/Public/ISourceControlWindowsModule.h
@@ -19,6 +19,14 @@ public:
 		return FModuleManager::LoadModuleChecked<ISourceControlWindowsModule>("SourceControlWindows");
 	}
 
+	static inline ISourceControlWindowsModule* TryGet()
+	{
+		return FModuleManager::GetModulePtr<ISourceControlWindowsModule>("SourceControlWindows");
+	}
+
 	virtual void ShowChangelistsTab() = 0;
 	virtual bool CanShowChangelistsTab() const = 0;
+
+	DECLARE_EVENT_OneParam(ISourceControlWindowsModule, FChangelistFileDoubleClickedEvent, const FString&);
+	virtual FChangelistFileDoubleClickedEvent& OnChangelistFileDoubleClicked() = 0;
 };
diff --git a/Engine/Source/Editor/SourceControlWindows/Public/SourceControlMenuContext.h b/Engine/Source/Editor/SourceControlWindows/Public/SourceControlMenuContext.h
new file mode 100644
--- /dev/null
+++ b/Engine/Source/Editor/SourceControlWindows/Public/SourceControlMenuContext.h
@@ -0,0 +1,27 @@
+// Copyright Epic Games, Inc. All Rights Reserved.
+#pragma once
+
+#include "CoreMinimal.h"
+#include "UObject/ObjectMacros.h"
+#include "UObject/Object.h"
+
+#include "SourceControlMenuContext.generated.h"
+
+/**
+ * Source control window menu context providing information for menu extenders.
+ */
+UCLASS()
+class SOURCECONTROLWINDOWS_API USourceControlMenuContext : public UObject
+{
+	GENERATED_BODY()
+public:
+
+	USourceControlMenuContext(const FObjectInitializer& ObjectInitializer)
+		: Super(ObjectInitializer)
+	{
+		
+	}
+
+	UPROPERTY()
+	TArray<FString> SelectedFiles;
+};
diff --git a/Engine/Source/Editor/UnrealEd/Classes/Editor/EditorEngine.h b/Engine/Source/Editor/UnrealEd/Classes/Editor/EditorEngine.h
--- a/Engine/Source/Editor/UnrealEd/Classes/Editor/EditorEngine.h
+++ b/Engine/Source/Editor/UnrealEd/Classes/Editor/EditorEngine.h
@@ -1037,10 +1037,11 @@ public:
 
 	/**
 	 * Moves all viewport cameras to focus on the provided bounding box.
-	 * @param	BoundingBox				Target box
-	 * @param	bActiveViewportOnly		If true, move/reorient only the active viewport.
+	 * @param	BoundingBox					Target box
+	 * @param	bActiveViewportOnly			If true, move/reorient only the active viewport.
+	 * @param	DrawDebubBoxTimeInSeconds	If greater than 0 a debug box is drawn representing the bounding box.It will be drawn for specified time.
 	 */
-	void MoveViewportCamerasToBox(const FBox& BoundingBox, bool bActiveViewportOnly) const;
+	void MoveViewportCamerasToBox(const FBox& BoundingBox, bool bActiveViewportOnly, float DrawDebugBoxTimeInSeconds = 0.f) const;
 
 	/** 
 	 * Snaps an element in a direction.  Optionally will align with the trace normal.
diff --git a/Engine/Source/Editor/UnrealEd/Private/EditorEngine.cpp b/Engine/Source/Editor/UnrealEd/Private/EditorEngine.cpp
--- a/Engine/Source/Editor/UnrealEd/Private/EditorEngine.cpp
+++ b/Engine/Source/Editor/UnrealEd/Private/EditorEngine.cpp
@@ -1021,6 +1021,7 @@ void UEditorEngine::Init(IEngineLoop* InEngineLoop)
 			TEXT("OutputLog"),
 			TEXT("SourceControl"),
 			TEXT("SourceControlWindows"),
+			TEXT("SourceControlWindowExtender"),
 			TEXT("UncontrolledChangelists"),
 			TEXT("TextureCompressor"),
 			TEXT("MeshUtilities"),
diff --git a/Engine/Source/Editor/UnrealEd/Private/EditorServer.cpp b/Engine/Source/Editor/UnrealEd/Private/EditorServer.cpp
--- a/Engine/Source/Editor/UnrealEd/Private/EditorServer.cpp
+++ b/Engine/Source/Editor/UnrealEd/Private/EditorServer.cpp
@@ -5021,7 +5021,7 @@ void UEditorEngine::MoveViewportCamerasToElement(const UTypedElementSelectionSet
 	MoveViewportCamerasToBox(BoundingBox, bActiveViewportOnly);
 }
 
-void UEditorEngine::MoveViewportCamerasToBox(const FBox& BoundingBox, bool bActiveViewportOnly) const
+void UEditorEngine::MoveViewportCamerasToBox(const FBox& BoundingBox, bool bActiveViewportOnly, float DrawDebugBoxTimeInSeconds) const
 {
 	// Make sure we had at least one non-null actor in the array passed in.
 	if (BoundingBox.GetSize() != FVector::ZeroVector || BoundingBox.GetCenter() != FVector::ZeroVector)
@@ -5058,6 +5058,16 @@ void UEditorEngine::MoveViewportCamerasToBox(const FBox& BoundingBox, bool bActi
 					LinkedViewportClient->FocusViewportOnBox(BoundingBox);
 			}
 		}
+
+		if (DrawDebugBoxTimeInSeconds > 0.f)
+		{
+			const FWorldContext& Context = const_cast<UEditorEngine*>(this)->GetEditorWorldContext();
+			if (const UWorld* World = Context.World())
+			{
+				const FColor DarkGray(40, 40, 40);
+				DrawDebugBox(World, BoundingBox.GetCenter(), BoundingBox.GetExtent(), DarkGray, false, DrawDebugBoxTimeInSeconds);
+			}
+		}
 	}
 }
 
diff --git a/Engine/Source/Editor/UnrealEd/UnrealEd.Build.cs b/Engine/Source/Editor/UnrealEd/UnrealEd.Build.cs
--- a/Engine/Source/Editor/UnrealEd/UnrealEd.Build.cs
+++ b/Engine/Source/Editor/UnrealEd/UnrealEd.Build.cs
@@ -265,6 +265,7 @@ public class UnrealEd : ModuleRules
 				"VirtualTexturingEditor",
 				"WorldPartitionEditor",
 				"CSVtoSVG",
+				"SourceControlWindowExtender"
 			}
 		);
 
