Add missing #undef LOCTEXT_NAMESPACE

#rnx

diff --git a/Engine/Source/Editor/SourceControlWindowExtender/Private/SourceControlWindowExtenderModule.cpp b/Engine/Source/Editor/SourceControlWindowExtender/Private/SourceControlWindowExtenderModule.cpp
--- a/Engine/Source/Editor/SourceControlWindowExtender/Private/SourceControlWindowExtenderModule.cpp
+++ b/Engine/Source/Editor/SourceControlWindowExtender/Private/SourceControlWindowExtenderModule.cpp
@@ -300,3 +300,4 @@ void FSourceControlWindowExtenderModule::OnChangelistFileDoubleClicked(const FSt
 	}
 }
 
+#undef LOCTEXT_NAMESPACE
\ No newline at end of file
