Fix FWorldPartitionActorDesc::GetBounds deprecation warning.

#rnx

diff --git a/Engine/Source/Editor/SourceControlWindowExtender/Private/SourceControlWindowExtenderModule.cpp b/Engine/Source/Editor/SourceControlWindowExtender/Private/SourceControlWindowExtenderModule.cpp
--- a/Engine/Source/Editor/SourceControlWindowExtender/Private/SourceControlWindowExtenderModule.cpp
+++ b/Engine/Source/Editor/SourceControlWindowExtender/Private/SourceControlWindowExtenderModule.cpp
@@ -259,7 +259,7 @@ void FSourceControlWindowExtenderModule::FocusActors(const TArray<FAssetData>& A
 	{
 		if (TUniquePtr<FWorldPartitionActorDesc> ActorDesc = FWorldPartitionActorDescUtils::GetActorDescriptorFromAssetData(ActorToFocus))
 		{
-			FocusBounds += ActorDesc->GetBounds();
+			FocusBounds += ActorDesc->GetEditorBounds();
 		}
 	}
 
